Improving SteamController support
- Enabled for Windows
- Added support for loading the proper Steamworks DLL
- Adding force feedback

diff --git a/Engine/Plugins/Runtime/Steam/SteamController/Source/SteamController/Private/SteamController.cpp b/Engine/Plugins/Runtime/Steam/SteamController/Source/SteamController/Private/SteamController.cpp
--- a/Engine/Plugins/Runtime/Steam/SteamController/Source/SteamController/Private/SteamController.cpp
+++ b/Engine/Plugins/Runtime/Steam/SteamController/Source/SteamController/Private/SteamController.cpp
@@ -6,6 +6,9 @@ DEFINE_LOG_CATEGORY_STATIC(LogSteamController, Log, All);
 
 #if WITH_STEAM_CONTROLLER
 
+/** Name of the current Steam SDK version in use (matches directory name) */
+#define STEAM_SDK_VER TEXT("Steamv130")
+
 /** @todo - do something about this define */
 #ifndef MAX_STEAM_CONTROLLERS
 	#define MAX_STEAM_CONTROLLERS 8
@@ -30,6 +33,36 @@ namespace
 	}
 };
 
+// Support function to load the proper version of the Steamworks library
+bool LoadSteamModule()
+{
+#if PLATFORM_WINDOWS
+#if PLATFORM_64BITS
+	void* SteamDLLHandle = nullptr;
+
+	FString RootSteamPath = FPaths::EngineDir() / FString::Printf(TEXT("Binaries/ThirdParty/Steamworks/%s/Win64/"), STEAM_SDK_VER); 
+	FPlatformProcess::PushDllDirectory(*RootSteamPath);
+	SteamDLLHandle = FPlatformProcess::GetDllHandle(*(RootSteamPath + "steam_api64.dll"));
+	FPlatformProcess::PopDllDirectory(*RootSteamPath);
+#else
+	FString RootSteamPath = FPaths::EngineDir() / FString::Printf(TEXT("Binaries/ThirdParty/Steamworks/%s/Win32/"), STEAM_SDK_VER); 
+	FPlatformProcess::PushDllDirectory(*RootSteamPath);
+	SteamDLLHandle = FPlatformProcess::GetDllHandle(*(RootSteamPath + "steam_api.dll"));
+	FPlatformProcess::PopDllDirectory(*RootSteamPath);
+#endif
+#elif PLATFORM_MAC
+	SteamDLLHandle = FPlatformProcess::GetDllHandle(TEXT("libsteam_api.dylib"));
+#endif	//PLATFORM_WINDOWS
+
+	if (!SteamDLLHandle)
+	{
+		UE_LOG(LogSteamController, Warning, TEXT("Failed to load Steam library."));
+		return false;
+	}
+
+	return true;
+}
+
 class FSteamController : public IInputDevice
 {
 
@@ -38,8 +71,17 @@ public:
 	FSteamController(const TSharedRef< FGenericApplicationMessageHandler >& InMessageHandler)
 		: MessageHandler(InMessageHandler)
 	{
-		// [RCL] 2015-01-23 FIXME: move to some other code than constructor so we can handle failures more gracefully
-		if (SteamController() != nullptr)
+		// Attempt to load the Steam Library
+		if (!LoadSteamModule())
+		{
+			return;
+		}
+
+		// Initialize the API, so we can start calling SteamController functions
+		bool bAPIInitialized = SteamAPI_Init();
+
+			// [RCL] 2015-01-23 FIXME: move to some other code than constructor so we can handle failures more gracefully
+		if (bAPIInitialized && (SteamController() != nullptr))
 		{
 			FString PluginsDir = FPaths::EnginePluginsDir();
 			FString ContentDir = FPaths::Combine(*PluginsDir, TEXT("Runtime"), TEXT("Steam"), TEXT("SteamController"), TEXT("Content"));
@@ -214,15 +256,53 @@ public:
 			}
 		}
 	}
-
+	
 	void SetChannelValue(int32 ControllerId, FForceFeedbackChannelType ChannelType, float Value) override
 	{
-		// @todo: implement
+		// Skip unless this is the large channels, which we consider to be the only SteamController feedback channels
+		if ((ChannelType != FForceFeedbackChannelType::LEFT_LARGE) && (ChannelType != FForceFeedbackChannelType::RIGHT_LARGE))
+		{
+			return;
+		}
+
+		if ((ControllerId >= 0) && (ControllerId < MAX_STEAM_CONTROLLERS))
+		{
+			FControllerState& ControllerState = ControllerStates[ControllerId];
+
+			ControllerState.VibeValues.LeftLarge = Value;
+
+			UpdateVibration(ControllerId);
+		}
 	}
 
 	void SetChannelValues(int32 ControllerId, const FForceFeedbackValues &Values) override
 	{
-		// @todo: implement
+		if ((ControllerId >= 0) && (ControllerId < MAX_STEAM_CONTROLLERS))
+		{
+			FControllerState& ControllerState = ControllerStates[ControllerId];
+			ControllerState.VibeValues = Values;
+
+			UpdateVibration(ControllerId);
+		}
+	}
+
+	void UpdateVibration(int32 ControllerId)
+	{
+		const FControllerState& ControllerState = ControllerStates[ControllerId];
+		ISteamController* const Controller = SteamController();
+
+		// Map the float values from [0,1] to be more reasonable values for the SteamController.  The docs say that [100,2000] are reasonable values
+ 		const float LeftIntensity = FMath::Clamp(ControllerState.VibeValues.LeftLarge * 2000.f, 0.f, 2000.f);
+		if (LeftIntensity > 0.f)
+		{
+			Controller->TriggerHapticPulse(ControllerId, ESteamControllerPad::k_ESteamControllerPad_Left, LeftIntensity);
+		}
+ 
+ 		const float RightIntensity = FMath::Clamp(ControllerState.VibeValues.RightLarge * 2000.f, 0.f, 2000.f);
+		if (RightIntensity > 0.f)
+		{
+			Controller->TriggerHapticPulse(ControllerId, ESteamControllerPad::k_ESteamControllerPad_Right, RightIntensity);
+		}
 	}
 
 	virtual void SetMessageHandler( const TSharedRef< FGenericApplicationMessageHandler >& InMessageHandler ) override
@@ -261,6 +341,9 @@ private:
 
 		/** Next time a repeat event should be generated for each button */
 		double NextRepeatTime[MAX_NUM_CONTROLLER_BUTTONS];
+
+		/** Values for force feedback on this controller.  We only consider the LEFT_LARGE channel for SteamControllers */
+		FForceFeedbackValues VibeValues;
 	}; 
 
 	/** Controller states */
@@ -286,7 +369,6 @@ class FSteamControllerPlugin : public IInputDeviceModule
 	virtual TSharedPtr< class IInputDevice > CreateInputDevice(const TSharedRef< FGenericApplicationMessageHandler >& InMessageHandler) override
 	{
 #if WITH_STEAM_CONTROLLER
-		// @todo: load vdf here
 		return TSharedPtr< class IInputDevice >(new FSteamController(InMessageHandler));
 #else
 		return nullptr;
diff --git a/Engine/Plugins/Runtime/Steam/SteamController/SteamController.uplugin b/Engine/Plugins/Runtime/Steam/SteamController/SteamController.uplugin
--- a/Engine/Plugins/Runtime/Steam/SteamController/SteamController.uplugin
+++ b/Engine/Plugins/Runtime/Steam/SteamController/SteamController.uplugin
@@ -15,7 +15,7 @@
 		{
 			"Name" : "SteamController",
 			"Type" : "Runtime",
-			"WhitelistPlatforms" : [ "Linux" ]
+			"WhitelistPlatforms" : [ "Win64", "Linux" ]
 		}
 	]
 }
