Fix SteamController for Mac and Linux.

diff --git a/Engine/Plugins/Runtime/Steam/SteamController/Source/SteamController/Private/SteamController.cpp b/Engine/Plugins/Runtime/Steam/SteamController/Source/SteamController/Private/SteamController.cpp
--- a/Engine/Plugins/Runtime/Steam/SteamController/Source/SteamController/Private/SteamController.cpp
+++ b/Engine/Plugins/Runtime/Steam/SteamController/Source/SteamController/Private/SteamController.cpp
@@ -51,7 +51,9 @@ bool LoadSteamModule()
 	FPlatformProcess::PopDllDirectory(*RootSteamPath);
 #endif
 #elif PLATFORM_MAC
-	SteamDLLHandle = FPlatformProcess::GetDllHandle(TEXT("libsteam_api.dylib"));
+	void* SteamDLLHandle = FPlatformProcess::GetDllHandle(TEXT("libsteam_api.dylib"));
+#elif PLATFORM_LINUX
+	void* SteamDLLHandle = FPlatformProcess::GetDllHandle(TEXT("libsteam_api.so"));
 #endif	//PLATFORM_WINDOWS
 
 	if (!SteamDLLHandle)
