Disable crazy warnings that claim that standard C library is "deprecated".

diff --git a/Engine/Plugins/Runtime/Steam/SteamController/Source/SteamController/Private/SteamControllerPrivatePCH.h b/Engine/Plugins/Runtime/Steam/SteamController/Source/SteamController/Private/SteamControllerPrivatePCH.h
--- a/Engine/Plugins/Runtime/Steam/SteamController/Source/SteamController/Private/SteamControllerPrivatePCH.h
+++ b/Engine/Plugins/Runtime/Steam/SteamController/Source/SteamController/Private/SteamControllerPrivatePCH.h
@@ -6,7 +6,18 @@
 #if WITH_STEAM_CONTROLLER
 
 #include "InputDevice.h"
+
+// Disable crazy warnings that claim that standard C library is "deprecated".
+#ifdef _MSC_VER
+#pragma warning(push)
+#pragma warning(disable:4996)
+#endif
+
 #include <steam/steam_api.h>
 
+#ifdef _MSC_VER
+#pragma warning(pop)
+#endif
+
 #endif // WITH_STEAM_CONTROLLER
 
