Crash when setting force feedback while SteamController plugin is enabled. This was caused by setting a force feedback to a SteamController device while none has been initialized (I belive the plugin enabled this). Solution is to check whether a gamepad of that type is attached before doing this. It's possible that we'll have to further implement IsGamepadAttached in other devices to conform with IInputDevice interface.

diff --git a/Engine/Plugins/Runtime/Steam/SteamController/Source/SteamController/Private/SteamController.cpp b/Engine/Plugins/Runtime/Steam/SteamController/Source/SteamController/Private/SteamController.cpp
--- a/Engine/Plugins/Runtime/Steam/SteamController/Source/SteamController/Private/SteamController.cpp
+++ b/Engine/Plugins/Runtime/Steam/SteamController/Source/SteamController/Private/SteamController.cpp
@@ -80,8 +80,10 @@ class FSteamController : public IInputDevice
 
 public:
 
-	FSteamController(const TSharedRef< FGenericApplicationMessageHandler >& InMessageHandler)
-		: MessageHandler(InMessageHandler)
+	FSteamController(const TSharedRef< FGenericApplicationMessageHandler >& InMessageHandler) :
+		MessageHandler(InMessageHandler), 
+		bSteamAPIInitialized(false), 
+		bSteamControllerInitialized(false)
 	{
 		// Attempt to load the Steam Library
 		if (!LoadSteamModule())
@@ -90,20 +92,20 @@ public:
 		}
 
 		// Initialize the API, so we can start calling SteamController functions
-		bool bAPIInitialized = SteamAPI_Init();
+		bSteamAPIInitialized = SteamAPI_Init();
 
 			// [RCL] 2015-01-23 FIXME: move to some other code than constructor so we can handle failures more gracefully
-		if (bAPIInitialized && (SteamController() != nullptr))
+		if (bSteamAPIInitialized && (SteamController() != nullptr))
 		{
 			FString PluginsDir = FPaths::EnginePluginsDir();
 			FString ContentDir = FPaths::Combine(*PluginsDir, TEXT("Runtime"), TEXT("Steam"), TEXT("SteamController"), TEXT("Content"));
 			FString VdfPath = FPaths::ConvertRelativePathToFull(FPaths::Combine(*ContentDir, TEXT("Controller.vdf")));
 
-			bool bInited = SteamController()->Init(TCHAR_TO_ANSI(*VdfPath));
-			UE_LOG(LogSteamController, Log, TEXT("SteamController %s initialized with vdf file '%s'."), bInited ? TEXT("could not be") : TEXT("has been"), *VdfPath);
+			bSteamControllerInitialized = SteamController()->Init(TCHAR_TO_ANSI(*VdfPath));
+			UE_LOG(LogSteamController, Log, TEXT("SteamController %s initialized with vdf file '%s'."), bSteamAPIInitialized ? TEXT("could not be") : TEXT("has been"), *VdfPath);
 
 			// [RCL] 2014-05-05 FIXME: disable when could not init?
-			if (bInited)
+			if (bSteamControllerInitialized)
 			{
 				FMemory::Memzero(ControllerStates, sizeof(ControllerStates));
 			}
@@ -300,8 +302,14 @@ public:
 
 	void UpdateVibration(int32 ControllerId)
 	{
-		const FControllerState& ControllerState = ControllerStates[ControllerId];
+		// make sure there is a valid device for this controller
 		ISteamController* const Controller = SteamController();
+		if (Controller == nullptr || IsGamepadAttached() == false)
+		{
+			return;
+		}
+
+		const FControllerState& ControllerState = ControllerStates[ControllerId];
 
 		// Map the float values from [0,1] to be more reasonable values for the SteamController.  The docs say that [100,2000] are reasonable values
  		const float LeftIntensity = FMath::Clamp(ControllerState.VibeValues.LeftLarge * 2000.f, 0.f, 2000.f);
@@ -327,6 +335,11 @@ public:
 		return false;
 	}
 
+	virtual bool IsGamepadAttached() const override
+	{
+		return (bSteamAPIInitialized && bSteamControllerInitialized);
+	}
+
 private:
 
 	struct FControllerState
@@ -372,6 +385,12 @@ private:
 
 	/** handler to send all messages to */
 	TSharedRef<FGenericApplicationMessageHandler> MessageHandler;
+
+	/** SteamAPI initialized **/
+	bool bSteamAPIInitialized;
+
+	/** SteamController initialized **/
+	bool bSteamControllerInitialized;
 };
 
 #endif // WITH_STEAM_CONTROLLER
diff --git a/Engine/Plugins/Runtime/Steam/SteamVR/Source/SteamVRController/Private/SteamVRController.cpp b/Engine/Plugins/Runtime/Steam/SteamVR/Source/SteamVRController/Private/SteamVRController.cpp
--- a/Engine/Plugins/Runtime/Steam/SteamVR/Source/SteamVRController/Private/SteamVRController.cpp
+++ b/Engine/Plugins/Runtime/Steam/SteamVR/Source/SteamVRController/Private/SteamVRController.cpp
@@ -409,6 +409,24 @@ public:
 		return false;
 	}
 
+	virtual bool IsGamepadAttached() const override
+	{
+		FSteamVRHMD* SteamVRSystem = GetSteamVRHMD();
+
+		if (SteamVRSystem != nullptr)
+		{
+			// Check if at least one motion controller is tracked
+			// Only need to check for at least one player (player index 0)
+			int32 PlayerIndex = 0;
+			ETrackingStatus LeftHandTrackingStatus = GetControllerTrackingStatus(PlayerIndex, EControllerHand::Left);
+			ETrackingStatus RightHandTrackingStatus = GetControllerTrackingStatus(PlayerIndex, EControllerHand::Left);
+
+			return LeftHandTrackingStatus == ETrackingStatus::Tracked || RightHandTrackingStatus == ETrackingStatus::Tracked;
+		}
+
+		return false;
+	}
+
 private:
 
 	inline vr::IVRSystem* GetVRSystem()
diff --git a/Engine/Source/Runtime/Core/Private/Windows/WindowsApplication.cpp b/Engine/Source/Runtime/Core/Private/Windows/WindowsApplication.cpp
--- a/Engine/Source/Runtime/Core/Private/Windows/WindowsApplication.cpp
+++ b/Engine/Source/Runtime/Core/Private/Windows/WindowsApplication.cpp
@@ -1908,7 +1908,12 @@ void FWindowsApplication::SetForceFeedbackChannelValues(int32 ControllerId, cons
 	// send vibration to externally-implemented devices
 	for( auto DeviceIt = ExternalInputDevices.CreateIterator(); DeviceIt; ++DeviceIt )
 	{
-		(*DeviceIt)->SetChannelValues(ControllerId, Values);
+		// *N.B 06/20/2016*: Ideally, we would want to use GetHapticDevice instead
+		// but they're not implemented for SteamController and SteamVRController
+		if ((*DeviceIt)->IsGamepadAttached()) 
+		{
+			(*DeviceIt)->SetChannelValues(ControllerId, Values);
+		}
 	}
 }
 
