Fix for wrongly changed ARRAY_COUNT macro in CL# 9279060.

diff --git a/Engine/Plugins/Runtime/Steam/SteamController/Source/SteamController/Private/SteamControllerPrivate.h b/Engine/Plugins/Runtime/Steam/SteamController/Source/SteamController/Private/SteamControllerPrivate.h
--- a/Engine/Plugins/Runtime/Steam/SteamController/Source/SteamController/Private/SteamControllerPrivate.h
+++ b/Engine/Plugins/Runtime/Steam/SteamController/Source/SteamController/Private/SteamControllerPrivate.h
@@ -14,8 +14,9 @@
 #pragma warning(disable:4996)
 #endif
 
-#pragma push_macro("UE_ARRAY_COUNT")
-#undef UE_ARRAY_COUNT
+// These can be removed when the deprecated ARRAY_COUNT is removed.
+#pragma push_macro("ARRAY_COUNT")
+#undef ARRAY_COUNT
 
 #if STEAMSDK_FOUND == 0
 #error Steam SDK not located.  Expected to be found in Engine/Source/ThirdParty/Steamworks/{SteamVersion}
@@ -32,7 +33,7 @@
  	MSVC_PRAGMA( warning( pop ) )
 #endif	// USING_CODE_ANALYSIS
 
-#pragma pop_macro("UE_ARRAY_COUNT")
+#pragma pop_macro("ARRAY_COUNT")
 
 #ifdef _MSC_VER
 #pragma warning(pop)
