Update SteamController to utilize the new version of OnControll functions instead of the soon to be deprecated versions

#rnx

diff --git a/Engine/Plugins/Runtime/Steam/SteamController/Source/SteamController/Private/SteamController.cpp b/Engine/Plugins/Runtime/Steam/SteamController/Source/SteamController/Private/SteamController.cpp
--- a/Engine/Plugins/Runtime/Steam/SteamController/Source/SteamController/Private/SteamController.cpp
+++ b/Engine/Plugins/Runtime/Steam/SteamController/Source/SteamController/Private/SteamController.cpp
@@ -13,6 +13,7 @@
 #include "ISteamControllerPlugin.h"
 #include "SteamControllerPrivate.h"
 #include "SteamSharedModule.h"
+#include "HAL/PlatformApplicationMisc.h"
 #include "GameFramework/InputSettings.h"
 
 DEFINE_LOG_CATEGORY_STATIC(LogSteamController, Log, All);
@@ -146,23 +147,27 @@ public:
 			static FString ControllerName(TEXT("SteamController"));
 			FInputDeviceScope InputScope(this, SystemName, i, ControllerName);
 
+			FPlatformUserId UserId = FPlatformMisc::GetPlatformUserForUserIndex(i);
+			FInputDeviceId DeviceId = INPUTDEVICEID_NONE;
+			IPlatformInputDeviceMapper::Get().RemapControllerIdToPlatformUserAndDevice(i, UserId, DeviceId);
+
 			for (auto It = DigitalActionHandlesMap.CreateConstIterator(); It; ++It)
 			{
 				FName DigitalActionName = It.Key();
 				ControllerDigitalActionData_t DigitalActionData = SteamInput()->GetDigitalActionData(ControllerHandle, It.Value());
 				if (ControllerState.DigitalStatusMap[DigitalActionName] == false && DigitalActionData.bState)
 				{
- 					MessageHandler->OnControllerButtonPressed(DigitalNamesToKeysMap[DigitalActionName].GetFName(), i, false);
+ 					MessageHandler->OnControllerButtonPressed(DigitalNamesToKeysMap[DigitalActionName].GetFName(), UserId, DeviceId, false);
 					ControllerState.DigitalRepeatTimeMap[DigitalActionName] = FPlatformTime::Seconds() + ButtonRepeatDelay;
 				}
 				else if (ControllerState.DigitalStatusMap[DigitalActionName] == true && !DigitalActionData.bState)
 				{
-					MessageHandler->OnControllerButtonReleased(DigitalNamesToKeysMap[DigitalActionName].GetFName(), i, false);
+					MessageHandler->OnControllerButtonReleased(DigitalNamesToKeysMap[DigitalActionName].GetFName(), UserId, DeviceId, false);
 				}
 				else if (ControllerState.DigitalStatusMap[DigitalActionName] == true && DigitalActionData.bState && ControllerState.DigitalRepeatTimeMap[DigitalActionName] <= CurrentTime)
 				{
 					ControllerState.DigitalRepeatTimeMap[DigitalActionName] += ButtonRepeatDelay;
-					MessageHandler->OnControllerButtonPressed(DigitalNamesToKeysMap[DigitalActionName].GetFName(), i, true);
+					MessageHandler->OnControllerButtonPressed(DigitalNamesToKeysMap[DigitalActionName].GetFName(), UserId, DeviceId, true);
 				}
 
 				ControllerState.DigitalStatusMap[DigitalActionName] = DigitalActionData.bState;
@@ -182,38 +187,38 @@ public:
 				{
 					if (ControllerState.AnalogStatusMap[It.Key()].x != AnalogActionData.x)
 					{
-						MessageHandler->OnControllerAnalog(EKeys::Gamepad_LeftX.GetFName(), i, AnalogActionData.x);
+						MessageHandler->OnControllerAnalog(EKeys::Gamepad_LeftX.GetFName(), UserId, DeviceId, AnalogActionData.x);
 					}
 					
 					if (ControllerState.AnalogStatusMap[It.Key()].y != AnalogActionData.y)
 					{
-						MessageHandler->OnControllerAnalog(EKeys::Gamepad_LeftY.GetFName(), i, AnalogActionData.y);
+						MessageHandler->OnControllerAnalog(EKeys::Gamepad_LeftY.GetFName(), UserId, DeviceId, AnalogActionData.y);
 					}
 				}
 				else if (AxisNamesToKeysMap[AnalogActionName] == EKeys::Gamepad_RightX || AxisNamesToKeysMap[AnalogActionName] == EKeys::Gamepad_RightY)
 				{
 					if (ControllerState.AnalogStatusMap[It.Key()].x != AnalogActionData.x)
 					{
-						MessageHandler->OnControllerAnalog(EKeys::Gamepad_RightX.GetFName(), i, AnalogActionData.x);
+						MessageHandler->OnControllerAnalog(EKeys::Gamepad_RightX.GetFName(), UserId, DeviceId, AnalogActionData.x);
 					}
 
 					if (ControllerState.AnalogStatusMap[It.Key()].y != AnalogActionData.y)
 					{
-						MessageHandler->OnControllerAnalog(EKeys::Gamepad_RightY.GetFName(), i, AnalogActionData.y);
+						MessageHandler->OnControllerAnalog(EKeys::Gamepad_RightY.GetFName(), UserId, DeviceId, AnalogActionData.y);
 					}
 				}
 				else if (AxisNamesToKeysMap[AnalogActionName] == EKeys::Gamepad_LeftTriggerAxis)
 				{
 					if (ControllerState.AnalogStatusMap[It.Key()].x != AnalogActionData.x)
 					{
-						MessageHandler->OnControllerAnalog(EKeys::Gamepad_LeftTriggerAxis.GetFName(), i, AnalogActionData.x);
+						MessageHandler->OnControllerAnalog(EKeys::Gamepad_LeftTriggerAxis.GetFName(), UserId, DeviceId, AnalogActionData.x);
 					}
 				}
 				else if (AxisNamesToKeysMap[AnalogActionName] == EKeys::Gamepad_RightTriggerAxis)
 				{
 					if (ControllerState.AnalogStatusMap[It.Key()].x != AnalogActionData.x)
 					{
-						MessageHandler->OnControllerAnalog(EKeys::Gamepad_RightTriggerAxis.GetFName(), i, AnalogActionData.x);
+						MessageHandler->OnControllerAnalog(EKeys::Gamepad_RightTriggerAxis.GetFName(), UserId, DeviceId, AnalogActionData.x);
 					}
 				}
 				ControllerState.AnalogStatusMap[AnalogActionName] = AnalogActionData;
