Fix for no unity build

diff --git a/Engine/Plugins/Runtime/Steam/SteamController/Source/SteamController/Private/SteamController.cpp b/Engine/Plugins/Runtime/Steam/SteamController/Source/SteamController/Private/SteamController.cpp
--- a/Engine/Plugins/Runtime/Steam/SteamController/Source/SteamController/Private/SteamController.cpp
+++ b/Engine/Plugins/Runtime/Steam/SteamController/Source/SteamController/Private/SteamController.cpp
@@ -6,6 +6,7 @@
 #include "Misc/Paths.h"
 #include "Misc/ConfigCacheIni.h"
 #include "GenericPlatform/GenericApplicationMessageHandler.h"
+#include "GenericPlatform/GenericPlatformInputDeviceMapper.h"
 #include "Modules/ModuleManager.h"
 #include "GenericPlatform/IInputInterface.h"
 #include "IInputDevice.h"
