Fix a race condition crash on shutdown of the SteamController plugin.

diff --git a/Engine/Plugins/Runtime/Steam/SteamController/Source/SteamController/Private/SteamController.cpp b/Engine/Plugins/Runtime/Steam/SteamController/Source/SteamController/Private/SteamController.cpp
--- a/Engine/Plugins/Runtime/Steam/SteamController/Source/SteamController/Private/SteamController.cpp
+++ b/Engine/Plugins/Runtime/Steam/SteamController/Source/SteamController/Private/SteamController.cpp
@@ -113,7 +113,7 @@ public:
 
 	virtual ~FSteamController()
 	{
-		if (SteamInput() != nullptr)
+		if (FSteamSharedModule::IsAvailable() && FSteamSharedModule::Get().AreSteamDllsLoaded() && SteamInput() != nullptr)
 		{
 			SteamInput()->Shutdown();
 		}
